Remove DistillEngine.xml as obsolete.

#codeview:bob.tellez
Do you know if all Distill stuff can be deleted? I searched UE4 and couldn't find any external references to that class.

diff --git a/Engine/Build/DistillEngine.xml b/Engine/Build/DistillEngine.xml
deleted file mode 100644
diff --git a/Engine/Source/Programs/Distill/Private/Distill.cs b/Engine/Source/Programs/Distill/Private/Distill.cs
--- a/Engine/Source/Programs/Distill/Private/Distill.cs
+++ b/Engine/Source/Programs/Distill/Private/Distill.cs
@@ -53,8 +53,8 @@ namespace Tools.Distill
 		public List<string> NotLanguages = null;
 		/// <summary>The name of the set of tags to manipulate files</summary>
 		public string TagSet = "Default";
-		/// <summary>The default xml file to load</summary>
-		public string DistillSettings = "../../Build/DistillEngine.xml";
+		/// <summary>The xml file to load</summary>
+		public string DistillSettings;
 
 		/// <summary></summary>
 		public bool bAuthenticate = false;
